SoundMod plugin

diff --git a/Engine/Plugins/Runtime/SoundMod/SoundMod.uplugin b/Engine/Plugins/Runtime/SoundMod/SoundMod.uplugin
new file mode 100644
--- /dev/null
+++ b/Engine/Plugins/Runtime/SoundMod/SoundMod.uplugin
@@ -0,0 +1,25 @@
+{
+    "PluginFileVersion" : 1,
+	
+	"FriendlyName" : "Sound Module Plugin",
+	"Version" : 1,
+	"VersionName" : "1.0",
+	"CreatedBy" : "Epic Games, Inc.",
+	"CreatedByURL" : "http://epicgames.com",
+	"MinEngineVersion" : 1579795,
+	"Description" : "Supports playback of ProTracker (MOD),  Scream Tracker 3 (S3M), Fast Tracker II (XM), and Impulse Tracker (IT) files.",
+	"CategoryPath" : "Audio.Playback",
+	
+    "Modules" :
+	[
+		{
+			"Name" : "SoundMod",
+			"Type" : "Runtime",
+			"WhitelistPlatforms" : [ "Win64", "Win32", "Mac" ]
+		},
+		{
+			"Name" : "SoundModImporter",
+			"Type" : "Editor"
+		}
+	]
+}
\ No newline at end of file
diff --git a/Engine/Plugins/Runtime/SoundMod/Source/SoundMod/Classes/SoundMod.h b/Engine/Plugins/Runtime/SoundMod/Source/SoundMod/Classes/SoundMod.h
new file mode 100644
--- /dev/null
+++ b/Engine/Plugins/Runtime/SoundMod/Source/SoundMod/Classes/SoundMod.h
@@ -0,0 +1,39 @@
+// Copyright 1998-2014 Epic Games, Inc. All Rights Reserved.
+
+#pragma once
+
+/** 
+ * A sound module file
+ */
+
+#include "SoundMod.generated.h"
+
+UCLASS(hidecategories=Object, MinimalAPI, BlueprintType)
+class USoundMod : public USoundBase
+{
+	GENERATED_UCLASS_BODY()
+
+	/** If set, when played directly (not through a sound cue) the nid will be played looping. */
+	UPROPERTY(EditAnywhere, Category=SoundMod)
+	uint32 bLooping:1;
+
+	/** The mod file data */
+	FByteBulkData				RawData;
+
+private:
+
+	/** Memory containing the data copied from the compressed bulk data */
+	uint8*	ResourceData;
+
+public:	
+	// Begin UObject interface. 
+	virtual void Serialize(FArchive& Ar) OVERRIDE;
+	// End UObject interface. 
+
+	// Begin USoundBase interface.
+	virtual bool IsPlayable() const OVERRIDE;
+	virtual void Parse(class FAudioDevice* AudioDevice, const UPTRINT NodeWaveInstanceHash, FActiveSound& ActiveSound, const FSoundParseParameters& ParseParams, TArray<FWaveInstance*>& WaveInstances) OVERRIDE;
+	virtual float GetMaxAudibleDistance() OVERRIDE;
+	// End USoundBase interface.
+};
+
diff --git a/Engine/Plugins/Runtime/SoundMod/Source/SoundMod/Classes/SoundModWave.h b/Engine/Plugins/Runtime/SoundMod/Source/SoundMod/Classes/SoundModWave.h
new file mode 100644
--- /dev/null
+++ b/Engine/Plugins/Runtime/SoundMod/Source/SoundMod/Classes/SoundModWave.h
@@ -0,0 +1,25 @@
+// Copyright 1998-2014 Epic Games, Inc. All Rights Reserved.
+
+/**
+* Playable sound object for wave files that are generated from a Sound Mod
+*/
+
+#pragma once
+#include "xmp.h"
+#include "SoundModWave.generated.h"
+
+UCLASS()
+class USoundModWave : public USoundWaveStreaming
+{
+	GENERATED_UCLASS_BODY()
+
+	UPROPERTY()
+	class USoundMod* SoundMod;
+
+	xmp_context xmpContext;
+
+	// Begin USoundWave interface.
+	virtual int32 GeneratePCMData(uint8* PCMData, const int32 SamplesNeeded) OVERRIDE;
+	// End USoundWave interface.
+
+};
\ No newline at end of file
diff --git a/Engine/Plugins/Runtime/SoundMod/Source/SoundMod/Private/SoundMod.cpp b/Engine/Plugins/Runtime/SoundMod/Source/SoundMod/Private/SoundMod.cpp
new file mode 100644
--- /dev/null
+++ b/Engine/Plugins/Runtime/SoundMod/Source/SoundMod/Private/SoundMod.cpp
@@ -0,0 +1,64 @@
+// Copyright 1998-2014 Epic Games, Inc. All Rights Reserved.
+
+#include "SoundModPrivatePCH.h"
+#include "SoundDefinitions.h"
+
+USoundMod::USoundMod(const class FPostConstructInitializeProperties& PCIP)
+	: Super(PCIP)
+{
+}
+
+void USoundMod::Serialize(FArchive& Ar)
+{
+	Super::Serialize(Ar);
+
+	RawData.Serialize(Ar, this);
+}
+
+void USoundMod::Parse(class FAudioDevice* AudioDevice, const UPTRINT NodeWaveInstanceHash, FActiveSound& ActiveSound, const FSoundParseParameters& ParseParams, TArray<FWaveInstance*>& WaveInstances)
+{
+	FWaveInstance* WaveInstance = ActiveSound.FindWaveInstance(NodeWaveInstanceHash);
+
+	// Create a new WaveInstance if this SoundWave doesn't already have one associated with it.
+	if (WaveInstance == NULL)
+	{
+		const int32 SampleRate = 44100;
+
+		// Create a new wave instance and associate with the ActiveSound
+		WaveInstance = new FWaveInstance(&ActiveSound);
+		WaveInstance->WaveInstanceHash = NodeWaveInstanceHash;
+		ActiveSound.WaveInstances.Add(NodeWaveInstanceHash, WaveInstance);
+
+		// Create streaming wave object
+		USoundModWave* ModWave = NewObject<USoundModWave>();
+		ModWave->SampleRate = SampleRate;
+		ModWave->NumChannels = 2;
+		ModWave->Duration = INDEFINITELY_LOOPING_DURATION;
+		ModWave->bLooping = bLooping;
+
+		if (ResourceData == NULL)
+		{
+			RawData.GetCopy((void**)&ResourceData, true);
+		}
+
+		ModWave->xmpContext = xmp_create_context();
+		xmp_load_module_from_memory(ModWave->xmpContext, ResourceData, RawData.GetBulkDataSize());
+		xmp_start_player(ModWave->xmpContext, SampleRate, 0);
+
+		WaveInstance->WaveData = ModWave;
+	}
+
+	WaveInstance->WaveData->Parse(AudioDevice, NodeWaveInstanceHash, ActiveSound, ParseParams, WaveInstances);
+}
+
+bool USoundMod::IsPlayable() const
+{
+	return true;
+}
+
+float USoundMod::GetMaxAudibleDistance()
+{
+	return (AttenuationSettings ? AttenuationSettings->Attenuation.GetMaxDimension() : WORLD_MAX);
+}
+
+
diff --git a/Engine/Plugins/Runtime/SoundMod/Source/SoundMod/Private/SoundModPlugin.cpp b/Engine/Plugins/Runtime/SoundMod/Source/SoundMod/Private/SoundModPlugin.cpp
new file mode 100644
--- /dev/null
+++ b/Engine/Plugins/Runtime/SoundMod/Source/SoundMod/Private/SoundModPlugin.cpp
@@ -0,0 +1,31 @@
+// Copyright 1998-2014 Epic Games, Inc. All Rights Reserved.
+
+#include "SoundModPrivatePCH.h"
+
+class FSoundModPlugin : public ISoundModPlugin
+{
+public:
+	// IModuleInterface implementation
+	virtual void StartupModule() OVERRIDE;
+	virtual void ShutdownModule() OVERRIDE;
+	// End of IModuleInterface implementation
+};
+
+IMPLEMENT_MODULE( FSoundModPlugin, SoundMod )
+
+
+
+void FSoundModPlugin::StartupModule()
+{
+	// This code will execute after your module is loaded into memory (but after global variables are initialized, of course.)
+}
+
+
+void FSoundModPlugin::ShutdownModule()
+{
+	// This function may be called during shutdown to clean up your module.  For modules that support dynamic reloading,
+	// we call this function before unloading the module.
+}
+
+
+
diff --git a/Engine/Plugins/Runtime/SoundMod/Source/SoundMod/Private/SoundModPrivatePCH.h b/Engine/Plugins/Runtime/SoundMod/Source/SoundMod/Private/SoundModPrivatePCH.h
new file mode 100644
--- /dev/null
+++ b/Engine/Plugins/Runtime/SoundMod/Source/SoundMod/Private/SoundModPrivatePCH.h
@@ -0,0 +1,9 @@
+// Copyright 1998-2014 Epic Games, Inc. All Rights Reserved.
+
+#include "Engine.h"
+#include "SoundModPlugin.h"
+#include "SoundMod.h"
+#include "SoundModWave.h"
+
+// You should place include statements to your module's private header files here.  You only need to
+// add includes for headers that are used in most of your module's source files though.
diff --git a/Engine/Plugins/Runtime/SoundMod/Source/SoundMod/Private/SoundModWave.cpp b/Engine/Plugins/Runtime/SoundMod/Source/SoundMod/Private/SoundModWave.cpp
new file mode 100644
--- /dev/null
+++ b/Engine/Plugins/Runtime/SoundMod/Source/SoundMod/Private/SoundModWave.cpp
@@ -0,0 +1,49 @@
+// Copyright 1998-2014 Epic Games, Inc. All Rights Reserved.
+
+#include "SoundModPrivatePCH.h"
+
+USoundModWave::USoundModWave(const class FPostConstructInitializeProperties& PCIP)
+	: Super(PCIP)
+	, xmpContext(NULL)
+{
+}
+
+DECLARE_CYCLE_STAT(TEXT("Sound Mod Generate Data"), STAT_SoundModGeneratePCMData, STATGROUP_Audio );
+
+int32 USoundModWave::GeneratePCMData(uint8* PCMData, const int32 SamplesNeeded)
+{
+	SCOPE_CYCLE_COUNTER(STAT_SoundModGeneratePCMData);
+
+	int32 BytesGenerated = Super::GeneratePCMData(PCMData, SamplesNeeded);
+	int32 BytesDesired = (SamplesNeeded * sizeof(int16)) - BytesGenerated;
+
+	if (BytesDesired > 0)
+	{
+		while (xmp_play_frame(xmpContext) == 0)
+		{
+			xmp_frame_info xmpFrameInfo;
+			xmp_get_frame_info(xmpContext, &xmpFrameInfo);
+
+			const int32 BytesToCopy = FMath::Min<int32>(xmpFrameInfo.buffer_size, BytesDesired);
+
+			if (BytesToCopy > 0)
+			{
+				FMemory::Memcpy((void*)(PCMData + BytesGenerated), xmpFrameInfo.buffer, BytesToCopy);
+
+				BytesDesired -= BytesToCopy;
+				BytesGenerated += BytesToCopy;
+
+				if (BytesDesired == 0)
+				{
+					if (BytesToCopy != xmpFrameInfo.buffer_size)
+					{
+						QueueAudio((uint8*)xmpFrameInfo.buffer + BytesToCopy, xmpFrameInfo.buffer_size - BytesToCopy);
+					}
+					break;
+				}
+			}
+		}
+	}
+
+	return BytesGenerated;
+}
\ No newline at end of file
diff --git a/Engine/Plugins/Runtime/SoundMod/Source/SoundMod/Private/SoundNodeModPlayer.cpp b/Engine/Plugins/Runtime/SoundMod/Source/SoundMod/Private/SoundNodeModPlayer.cpp
new file mode 100644
--- /dev/null
+++ b/Engine/Plugins/Runtime/SoundMod/Source/SoundMod/Private/SoundNodeModPlayer.cpp
@@ -0,0 +1,98 @@
+// Copyright 1998-2014 Epic Games, Inc. All Rights Reserved.
+
+#include "SoundModPrivatePCH.h"
+#include "SoundDefinitions.h"
+#include "SoundNodeModPlayer.h"
+
+#define LOCTEXT_NAMESPACE "SoundNodeModPlayer"
+
+USoundNodeModPlayer::USoundNodeModPlayer(const class FPostConstructInitializeProperties& PCIP)
+	: Super(PCIP)
+{
+}
+
+void USoundNodeModPlayer::ParseNodes( FAudioDevice* AudioDevice, const UPTRINT NodeWaveInstanceHash, FActiveSound& ActiveSound, const FSoundParseParameters& ParseParams, TArray<FWaveInstance*>& WaveInstances )
+{
+	if (SoundMod)
+	{
+		// The SoundWave's bLooping is only for if it is directly referenced, so clear it
+		// in the case that it is being played from a player
+		bool bModIsLooping = SoundMod->bLooping;
+		SoundMod->bLooping = false;
+
+		if (bLooping)
+		{
+			FSoundParseParameters UpdatedParams = ParseParams;
+			UpdatedParams.bLooping = true;
+			SoundMod->Parse(AudioDevice, NodeWaveInstanceHash, ActiveSound, UpdatedParams, WaveInstances);
+		}
+		else
+		{
+			SoundMod->Parse(AudioDevice, NodeWaveInstanceHash, ActiveSound, ParseParams, WaveInstances);
+		}
+
+		SoundMod->bLooping = bModIsLooping;
+	}
+}
+
+float USoundNodeModPlayer::GetDuration()
+{
+	float Duration = 0.f;
+	if (SoundMod)
+	{
+		if (bLooping)
+		{
+			Duration = INDEFINITELY_LOOPING_DURATION;
+		}
+		else
+		{
+			Duration = SoundMod->Duration;
+		}
+	}
+	return Duration;
+}
+
+#if WITH_EDITOR
+FString USoundNodeModPlayer::GetTitle() const
+{
+	FText SoundModName;
+	if (SoundMod)
+	{
+		SoundModName = FText::FromString(SoundMod->GetFName().ToString());
+	}
+	else
+	{
+		SoundModName = LOCTEXT("NoSoundMod", "NONE");
+	}
+
+	FString Title;
+
+	if (bLooping)
+	{
+		FFormatNamedArguments Arguments;
+		Arguments.Add(TEXT("Description"), FText::FromString(Super::GetTitle()));
+		Arguments.Add(TEXT("SoundModName"), SoundModName);
+		Title = FText::Format(LOCTEXT("LoopingSoundWaveDescription", "Looping {Description} : {SoundModName}"), Arguments).ToString();
+	}
+	else
+	{
+		Title = Super::GetTitle() + FString(TEXT(" : ")) + SoundModName.ToString();
+	}
+
+	return Title;
+}
+#endif
+
+FString USoundNodeModPlayer::GetUniqueString() const
+{
+	return TEXT("");
+}
+
+// A Mod Player is the end of the chain and has no children
+int32 USoundNodeModPlayer::GetMaxChildNodes() const
+{
+	return 0;
+}
+
+
+#undef LOCTEXT_NAMESPACE
diff --git a/Engine/Plugins/Runtime/SoundMod/Source/SoundMod/Public/SoundModPlugin.h b/Engine/Plugins/Runtime/SoundMod/Source/SoundMod/Public/SoundModPlugin.h
new file mode 100644
--- /dev/null
+++ b/Engine/Plugins/Runtime/SoundMod/Source/SoundMod/Public/SoundModPlugin.h
@@ -0,0 +1,37 @@
+// Copyright 1998-2014 Epic Games, Inc. All Rights Reserved.
+
+#pragma once
+
+#include "ModuleManager.h"
+
+
+/**
+ * The public interface to this module
+ */
+class ISoundModPlugin : public IModuleInterface
+{
+
+public:
+
+	/**
+	 * Singleton-like access to this module's interface.  This is just for convenience!
+	 * Beware of calling this during the shutdown phase, though.  Your module might have been unloaded already.
+	 *
+	 * @return Returns singleton instance, loading the module on demand if needed
+	 */
+	static inline ISoundModPlugin& Get()
+	{
+		return FModuleManager::LoadModuleChecked<ISoundModPlugin>("SoundMod");
+	}
+
+	/**
+	 * Checks to see if this module is loaded and ready.  It is only valid to call Get() if IsAvailable() returns true.
+	 *
+	 * @return True if the module is loaded and ready to use
+	 */
+	static inline bool IsAvailable()
+	{
+		return FModuleManager::Get().IsModuleLoaded( "SoundMod" );
+	}
+};
+
diff --git a/Engine/Plugins/Runtime/SoundMod/Source/SoundMod/Public/SoundNodeModPlayer.h b/Engine/Plugins/Runtime/SoundMod/Source/SoundMod/Public/SoundNodeModPlayer.h
new file mode 100644
--- /dev/null
+++ b/Engine/Plugins/Runtime/SoundMod/Source/SoundMod/Public/SoundNodeModPlayer.h
@@ -0,0 +1,34 @@
+// Copyright 1998-2014 Epic Games, Inc. All Rights Reserved.
+
+
+#pragma once
+#include "Sound/SoundNode.h"
+#include "SoundNodeModPlayer.generated.h"
+
+/** 
+ * Sound node that contains a reference to the mod file to be played
+ */
+UCLASS(hidecategories=Object, editinlinenew, MinimalAPI, meta=( DisplayName="Mod Player" ))
+class USoundNodeModPlayer : public USoundNode
+{
+	GENERATED_UCLASS_BODY()
+
+	UPROPERTY(EditAnywhere, Category=ModPlayer)
+	USoundMod* SoundMod;
+
+	UPROPERTY(EditAnywhere, Category=ModPlayer)
+	uint32 bLooping:1;
+
+public:	
+	// Begin USoundNode Interface
+	virtual int32 GetMaxChildNodes() const OVERRIDE;
+	virtual float GetDuration() OVERRIDE;
+	virtual void ParseNodes( FAudioDevice* AudioDevice, const UPTRINT NodeWaveInstanceHash, FActiveSound& ActiveSound, const FSoundParseParameters& ParseParams, TArray<FWaveInstance*>& WaveInstances ) OVERRIDE;
+	virtual FString GetUniqueString() const OVERRIDE;
+#if WITH_EDITOR
+	virtual FString GetTitle() const OVERRIDE;
+#endif
+	// End USoundNode Interface
+
+};
+
diff --git a/Engine/Plugins/Runtime/SoundMod/Source/SoundMod/SoundMod.Build.cs b/Engine/Plugins/Runtime/SoundMod/Source/SoundMod/SoundMod.Build.cs
new file mode 100644
--- /dev/null
+++ b/Engine/Plugins/Runtime/SoundMod/Source/SoundMod/SoundMod.Build.cs
@@ -0,0 +1,29 @@
+// Copyright 1998-2014 Epic Games, Inc. All Rights Reserved.
+
+namespace UnrealBuildTool.Rules
+{
+    public class SoundMod : ModuleRules
+    {
+        public SoundMod(TargetInfo Target)
+        {
+            PublicDependencyModuleNames.AddRange(
+                new string[]
+				{
+					"Core",
+					"CoreUObject",
+                    "Engine",
+				}
+            );
+
+			Definitions.AddRange(
+				new string[] {
+					"BUILDING_STATIC"
+				}
+			);
+
+			// Link with managed Perforce wrapper assemblies
+			AddThirdPartyPrivateStaticDependencies(Target, "coremod");
+		
+		}
+    }
+}
diff --git a/Engine/Plugins/Runtime/SoundMod/Source/SoundModImporter/Classes/SoundModImporterFactory.h b/Engine/Plugins/Runtime/SoundMod/Source/SoundModImporter/Classes/SoundModImporterFactory.h
new file mode 100644
--- /dev/null
+++ b/Engine/Plugins/Runtime/SoundMod/Source/SoundModImporter/Classes/SoundModImporterFactory.h
@@ -0,0 +1,17 @@
+// Copyright 1998-2014 Epic Games, Inc. All Rights Reserved.
+
+#pragma once
+
+#include "SoundModImporterFactory.generated.h"
+
+// Imports a sound module file
+UCLASS()
+class USoundModImporterFactory : public UFactory
+{
+	GENERATED_UCLASS_BODY()
+
+	// UFactory interface
+	virtual UObject* FactoryCreateBinary(UClass* Class, UObject* InParent, FName Name, EObjectFlags Flags, UObject* Context, const TCHAR* Type, const uint8*& Buffer, const uint8* BufferEnd, FFeedbackContext* Warn) OVERRIDE;
+	// End of UFactory interface
+};
+
diff --git a/Engine/Plugins/Runtime/SoundMod/Source/SoundModImporter/Private/AssetTypeActions_SoundMod.cpp b/Engine/Plugins/Runtime/SoundMod/Source/SoundModImporter/Private/AssetTypeActions_SoundMod.cpp
new file mode 100644
--- /dev/null
+++ b/Engine/Plugins/Runtime/SoundMod/Source/SoundModImporter/Private/AssetTypeActions_SoundMod.cpp
@@ -0,0 +1,122 @@
+// Copyright 1998-2014 Epic Games, Inc. All Rights Reserved.
+
+#include "SoundModImporterPrivatePCH.h"
+#include "AssetTypeActions_SoundMod.h"
+
+// Copyright 1998-2014 Epic Games, Inc. All Rights Reserved.
+
+#define LOCTEXT_NAMESPACE "AssetTypeActions"
+
+UClass* FAssetTypeActions_SoundMod::GetSupportedClass() const
+{
+	return USoundMod::StaticClass();
+}
+
+void FAssetTypeActions_SoundMod::GetActions(const TArray<UObject*>& InObjects, FMenuBuilder& MenuBuilder)
+{
+	auto Sounds = GetTypedWeakObjectPtrs<USoundMod>(InObjects);
+
+	MenuBuilder.AddMenuEntry(
+		LOCTEXT("Sound_PlaySound", "Play"),
+		LOCTEXT("Sound_PlaySoundTooltip", "Plays the selected sound."),
+		FSlateIcon(),
+		FUIAction(
+		FExecuteAction::CreateSP(this, &FAssetTypeActions_SoundMod::ExecutePlaySound, Sounds),
+		FCanExecuteAction::CreateSP(this, &FAssetTypeActions_SoundMod::CanExecutePlayCommand, Sounds)
+		)
+		);
+
+	MenuBuilder.AddMenuEntry(
+		LOCTEXT("Sound_StopSound", "Stop"),
+		LOCTEXT("Sound_StopSoundTooltip", "Stops the selected sounds."),
+		FSlateIcon(),
+		FUIAction(
+		FExecuteAction::CreateSP(this, &FAssetTypeActions_SoundMod::ExecuteStopSound, Sounds),
+		FCanExecuteAction()
+		)
+		);
+}
+
+bool FAssetTypeActions_SoundMod::CanExecutePlayCommand(TArray<TWeakObjectPtr<USoundMod>> Objects) const
+{
+	return Objects.Num() == 1;
+}
+
+void FAssetTypeActions_SoundMod::AssetsActivated(const TArray<UObject*>& InObjects, EAssetTypeActivationMethod::Type ActivationType)
+{
+	if (ActivationType == EAssetTypeActivationMethod::Previewed)
+	{
+		USoundMod* TargetSound = NULL;
+
+		for (auto ObjIt = InObjects.CreateConstIterator(); ObjIt; ++ObjIt)
+		{
+			TargetSound = Cast<USoundMod>(*ObjIt);
+			if (TargetSound)
+			{
+				// Only target the first valid sound cue
+				break;
+			}
+		}
+
+		UAudioComponent* PreviewComp = GEditor->GetPreviewAudioComponent();
+		if (PreviewComp && PreviewComp->IsPlaying())
+		{
+			// Already previewing a sound, if it is the target cue then stop it, otherwise play the new one
+			if (!TargetSound || PreviewComp->Sound == TargetSound)
+			{
+				StopSound();
+			}
+			else
+			{
+				PlaySound(TargetSound);
+			}
+		}
+		else
+		{
+			// Not already playing, play the target sound cue if it exists
+			PlaySound(TargetSound);
+		}
+	}
+	else
+	{
+		FAssetTypeActions_Base::AssetsActivated(InObjects, ActivationType);
+	}
+}
+
+void FAssetTypeActions_SoundMod::ExecutePlaySound(TArray<TWeakObjectPtr<USoundMod>> Objects)
+{
+	for (auto ObjIt = Objects.CreateConstIterator(); ObjIt; ++ObjIt)
+	{
+		USoundMod* Sound = (*ObjIt).Get();
+		if (Sound)
+		{
+			// Only play the first valid sound
+			PlaySound(Sound);
+			break;
+		}
+	}
+}
+
+void FAssetTypeActions_SoundMod::ExecuteStopSound(TArray<TWeakObjectPtr<USoundMod>> Objects)
+{
+	StopSound();
+}
+
+void FAssetTypeActions_SoundMod::PlaySound(USoundMod* Sound)
+{
+	if (Sound)
+	{
+		GEditor->PlayPreviewSound(Sound);
+	}
+	else
+	{
+		StopSound();
+	}
+}
+
+void FAssetTypeActions_SoundMod::StopSound()
+{
+	GEditor->ResetPreviewAudioComponent();
+}
+
+#undef LOCTEXT_NAMESPACE
diff --git a/Engine/Plugins/Runtime/SoundMod/Source/SoundModImporter/Private/AssetTypeActions_SoundMod.h b/Engine/Plugins/Runtime/SoundMod/Source/SoundModImporter/Private/AssetTypeActions_SoundMod.h
new file mode 100644
--- /dev/null
+++ b/Engine/Plugins/Runtime/SoundMod/Source/SoundModImporter/Private/AssetTypeActions_SoundMod.h
@@ -0,0 +1,35 @@
+// Copyright 1998-2014 Epic Games, Inc. All Rights Reserved.
+
+#pragma once
+
+#include "AssetTypeActions_Base.h"
+
+class FAssetTypeActions_SoundMod: public FAssetTypeActions_Base
+{
+public:
+	// IAssetTypeActions Implementation
+	virtual FText GetName() const OVERRIDE{ return NSLOCTEXT("AssetTypeActions", "AssetTypeActions_SoundMod", "Sound Mod"); }
+	virtual FColor GetTypeColor() const OVERRIDE{ return FColor(255, 175, 0); }
+	virtual UClass* GetSupportedClass() const OVERRIDE;
+	virtual bool HasActions(const TArray<UObject*>& InObjects) const OVERRIDE{ return true; }
+	virtual void GetActions(const TArray<UObject*>& InObjects, FMenuBuilder& MenuBuilder) OVERRIDE;
+	virtual void AssetsActivated(const TArray<UObject*>& InObjects, EAssetTypeActivationMethod::Type ActivationType) OVERRIDE;
+	virtual uint32 GetCategories() OVERRIDE{ return EAssetTypeCategories::Sounds; }
+	virtual bool CanFilter() OVERRIDE{ return false; }
+
+private:
+	/** Handler for when PlaySound is selected */
+	void ExecutePlaySound(TArray<TWeakObjectPtr<USoundMod>> Objects);
+
+	/** Handler for when StopSound is selected */
+	void ExecuteStopSound(TArray<TWeakObjectPtr<USoundMod>> Objects);
+
+	/** Returns true if only one sound is selected to play */
+	bool CanExecutePlayCommand(TArray<TWeakObjectPtr<USoundMod>> Objects) const;
+
+	/** Plays the specified sound wave */
+	void PlaySound(USoundMod* Sound);
+
+	/** Stops any currently playing sounds */
+	void StopSound();
+};
\ No newline at end of file
diff --git a/Engine/Plugins/Runtime/SoundMod/Source/SoundModImporter/Private/SoundModImporterFactory.cpp b/Engine/Plugins/Runtime/SoundMod/Source/SoundModImporter/Private/SoundModImporterFactory.cpp
new file mode 100644
--- /dev/null
+++ b/Engine/Plugins/Runtime/SoundMod/Source/SoundModImporter/Private/SoundModImporterFactory.cpp
@@ -0,0 +1,129 @@
+// Copyright 1998-2014 Epic Games, Inc. All Rights Reserved.
+
+#include "SoundModImporterPrivatePCH.h"
+#include "AssetToolsModule.h"
+#include "AssetRegistryModule.h"
+#include "PackageTools.h"
+#include "SoundDefinitions.h"
+#include "xmp.h"
+
+//////////////////////////////////////////////////////////////////////////
+// USoundModImporterFactory
+
+static bool bSoundModFactorySuppressImportOverwriteDialog = false;
+
+USoundModImporterFactory::USoundModImporterFactory(const FPostConstructInitializeProperties& PCIP)
+	: Super(PCIP)
+{
+	bCreateNew = false;
+	SupportedClass = USoundMod::StaticClass();
+
+	bEditorImport = true;
+
+	Formats.Add(TEXT("mod;Protracker file"));
+	Formats.Add(TEXT("s3m;Scream Tracker 3 file"));
+	Formats.Add(TEXT("xm;Fast Tracker II file"));
+	Formats.Add(TEXT("it;Impulse Tracker file"));
+}
+
+UObject* USoundModImporterFactory::FactoryCreateBinary
+(
+	UClass*				Class,
+	UObject*			InParent,
+	FName				Name,
+	EObjectFlags		Flags,
+	UObject*			Context,
+	const TCHAR*		FileType,
+	const uint8*&		Buffer,
+	const uint8*		BufferEnd,
+	FFeedbackContext*	Warn
+)
+{
+	// if the sound already exists, remember the user settings
+	USoundMod* ExistingSound = FindObject<USoundMod>(InParent, *Name.ToString());
+
+	// TODO - Audio Threading. This needs to be sent to the audio device and wait on stopping the sounds
+	TArray<UAudioComponent*> ComponentsToRestart;
+	FAudioDevice* AudioDevice = GEngine->GetAudioDevice();
+	if (AudioDevice && ExistingSound)
+	{
+		// TODO: Generalize the stop sounds function
+		//AudioDevice->StopSoundsForReimport(ExistingSound, ComponentsToRestart);
+	}
+
+	bool bUseExistingSettings = bSoundModFactorySuppressImportOverwriteDialog;
+
+	if (ExistingSound && !bUseExistingSettings && !GIsAutomationTesting)
+	{
+		// Prompt the user for what to do if a 'To All' response wasn't already given.
+		if (OverwriteYesOrNoToAllState != EAppReturnType::YesAll && OverwriteYesOrNoToAllState != EAppReturnType::NoAll)
+		{
+			OverwriteYesOrNoToAllState = FMessageDialog::Open(EAppMsgType::YesNoYesAllNoAllCancel, FText::Format(
+				NSLOCTEXT("UnrealEd", "ImportedSoundAlreadyExists_F", "You are about to import '{0}' over an existing sound. Would you like to overwrite the existing settings?\n\nYes or Yes to All: Overwrite the existing settings.\nNo or No to All: Preserve the existing settings.\nCancel: Abort the operation."),
+				FText::FromName(Name)));
+		}
+
+		switch (OverwriteYesOrNoToAllState)
+		{
+			case EAppReturnType::Yes:
+			case EAppReturnType::YesAll:
+			{
+				// Overwrite existing settings
+				bUseExistingSettings = false;
+				break;
+			}
+			case EAppReturnType::No:
+			case EAppReturnType::NoAll:
+			{
+				// Preserve existing settings
+				bUseExistingSettings = true;
+				break;
+			}
+			default:
+			{
+				FEditorDelegates::OnAssetPostImport.Broadcast(this, NULL);
+				return NULL;
+			}
+		}
+	}
+
+	// Reset the flag back to false so subsequent imports are not suppressed unless the code explicitly suppresses it
+	bSoundModFactorySuppressImportOverwriteDialog = false;
+
+	TArray<uint8> RawModData;
+	RawModData.Empty(BufferEnd - Buffer);
+	RawModData.AddUninitialized(BufferEnd - Buffer);
+	FMemory::Memcpy(RawModData.GetData(), Buffer, RawModData.Num());
+
+	// TODO: Validate that this is actually a mod file
+	xmp_context xmpContext = xmp_create_context();
+	if (xmp_load_module_from_memory(xmpContext, RawModData.GetData(), RawModData.Num()) != 0)
+	{
+		return NULL;
+	}
+	xmp_module_info xmpModuleInfo;
+	xmp_get_module_info(xmpContext, &xmpModuleInfo);
+
+	// Use pre-existing sound if it exists and we want to keep settings,
+	// otherwise create new sound and import raw data.
+	USoundMod* Sound = (bUseExistingSettings && ExistingSound) ? ExistingSound : new(InParent, Name, Flags) USoundMod(FPostConstructInitializeProperties());
+
+	Sound->Duration = xmpModuleInfo.seq_data->duration / 1000.f;
+
+	xmp_release_module(xmpContext);
+	xmp_free_context(xmpContext);
+
+	Sound->RawData.Lock(LOCK_READ_WRITE);
+	void* LockedData = Sound->RawData.Realloc(BufferEnd - Buffer);
+	FMemory::Memcpy(LockedData, Buffer, BufferEnd - Buffer);
+	Sound->RawData.Unlock();
+
+	FEditorDelegates::OnAssetPostImport.Broadcast(this, Sound);
+
+	for (int32 ComponentIndex = 0; ComponentIndex < ComponentsToRestart.Num(); ++ComponentIndex)
+	{
+		ComponentsToRestart[ComponentIndex]->Play();
+	}
+
+	return Sound;
+}
\ No newline at end of file
diff --git a/Engine/Plugins/Runtime/SoundMod/Source/SoundModImporter/Private/SoundModImporterModule.cpp b/Engine/Plugins/Runtime/SoundMod/Source/SoundModImporter/Private/SoundModImporterModule.cpp
new file mode 100644
--- /dev/null
+++ b/Engine/Plugins/Runtime/SoundMod/Source/SoundModImporter/Private/SoundModImporterModule.cpp
@@ -0,0 +1,43 @@
+// Copyright 1998-2014 Epic Games, Inc. All Rights Reserved.
+
+#include "SoundModImporterPrivatePCH.h"
+#include "AssetTypeActions_SoundMod.h"
+
+//////////////////////////////////////////////////////////////////////////
+// FSoundModImporterModule
+
+class FSoundModImporterModule : public FDefaultModuleImpl
+{
+public:
+	/** Asset type actions for SoundMod assets.  Cached here so that we can unregister it during shutdown. */
+	TSharedPtr<FAssetTypeActions_SoundMod> SoundModAssetTypeActions;
+
+public:
+	virtual void StartupModule() OVERRIDE
+	{
+		// Register the asset type
+		SoundModAssetTypeActions = MakeShareable(new FAssetTypeActions_SoundMod);
+		FModuleManager::LoadModuleChecked<FAssetToolsModule>("AssetTools").Get().RegisterAssetTypeActions(SoundModAssetTypeActions.ToSharedRef());
+	}
+
+	virtual void ShutdownModule() OVERRIDE
+	{
+		if (!UObjectInitialized())
+		{
+			return;
+		}
+
+		// Only unregister if the asset tools module is loaded.  We don't want to forcibly load it during shutdown phase.
+		check(SoundModAssetTypeActions.IsValid());
+		if (FModuleManager::Get().IsModuleLoaded("AssetTools"))
+		{
+			FModuleManager::GetModuleChecked<FAssetToolsModule>("AssetTools").Get().UnregisterAssetTypeActions(SoundModAssetTypeActions.ToSharedRef());
+		}
+		SoundModAssetTypeActions.Reset();
+	}
+};
+
+//////////////////////////////////////////////////////////////////////////
+
+IMPLEMENT_MODULE(FSoundModImporterModule, SoundModImporter);
+DEFINE_LOG_CATEGORY(LogSoundModImporter);
diff --git a/Engine/Plugins/Runtime/SoundMod/Source/SoundModImporter/Private/SoundModImporterPrivatePCH.h b/Engine/Plugins/Runtime/SoundMod/Source/SoundModImporter/Private/SoundModImporterPrivatePCH.h
new file mode 100644
--- /dev/null
+++ b/Engine/Plugins/Runtime/SoundMod/Source/SoundModImporter/Private/SoundModImporterPrivatePCH.h
@@ -0,0 +1,10 @@
+// Copyright 1998-2014 Epic Games, Inc. All Rights Reserved.
+
+#pragma once
+
+#include "UnrealEd.h"
+#include "SoundMod.h"
+
+DECLARE_LOG_CATEGORY_EXTERN(LogSoundModImporter, Verbose, All);
+
+#include "SoundModImporterClasses.h"
diff --git a/Engine/Plugins/Runtime/SoundMod/Source/SoundModImporter/SoundModImporter.Build.cs b/Engine/Plugins/Runtime/SoundMod/Source/SoundModImporter/SoundModImporter.Build.cs
new file mode 100644
--- /dev/null
+++ b/Engine/Plugins/Runtime/SoundMod/Source/SoundModImporter/SoundModImporter.Build.cs
@@ -0,0 +1,41 @@
+// Copyright 1998-2014 Epic Games, Inc. All Rights Reserved.
+
+using UnrealBuildTool;
+
+public class SoundModImporter : ModuleRules
+{
+	public SoundModImporter(TargetInfo Target)
+	{
+		PrivateDependencyModuleNames.AddRange(
+			new string[] {
+				"Core",
+				"CoreUObject",
+				"Engine",
+				"Slate",
+				"SoundMod",
+				"UnrealEd"
+			});
+
+		PrivateIncludePathModuleNames.AddRange(
+			new string[] {
+				"AssetTools",
+				"AssetRegistry"
+			});
+
+		DynamicallyLoadedModuleNames.AddRange(
+			new string[] {
+				"AssetTools",
+				"AssetRegistry"
+			});
+
+		Definitions.AddRange(
+			new string[] {
+					"BUILDING_STATIC"
+				}
+		);
+
+		// Link with managed Perforce wrapper assemblies
+		AddThirdPartyPrivateStaticDependencies(Target, "coremod");
+	
+	}
+}
