Fixup OVERRIDE usages in SoundMod plugin

diff --git a/Engine/Plugins/Runtime/SoundMod/Source/SoundMod/Classes/SoundMod.h b/Engine/Plugins/Runtime/SoundMod/Source/SoundMod/Classes/SoundMod.h
--- a/Engine/Plugins/Runtime/SoundMod/Source/SoundMod/Classes/SoundMod.h
+++ b/Engine/Plugins/Runtime/SoundMod/Source/SoundMod/Classes/SoundMod.h
@@ -27,13 +27,13 @@ private:
 
 public:	
 	// Begin UObject interface. 
-	virtual void Serialize(FArchive& Ar) OVERRIDE;
+	virtual void Serialize(FArchive& Ar) override;
 	// End UObject interface. 
 
 	// Begin USoundBase interface.
-	virtual bool IsPlayable() const OVERRIDE;
-	virtual void Parse(class FAudioDevice* AudioDevice, const UPTRINT NodeWaveInstanceHash, FActiveSound& ActiveSound, const FSoundParseParameters& ParseParams, TArray<FWaveInstance*>& WaveInstances) OVERRIDE;
-	virtual float GetMaxAudibleDistance() OVERRIDE;
+	virtual bool IsPlayable() const override;
+	virtual void Parse(class FAudioDevice* AudioDevice, const UPTRINT NodeWaveInstanceHash, FActiveSound& ActiveSound, const FSoundParseParameters& ParseParams, TArray<FWaveInstance*>& WaveInstances) override;
+	virtual float GetMaxAudibleDistance() override;
 	// End USoundBase interface.
 };
 
diff --git a/Engine/Plugins/Runtime/SoundMod/Source/SoundMod/Classes/SoundModWave.h b/Engine/Plugins/Runtime/SoundMod/Source/SoundMod/Classes/SoundModWave.h
--- a/Engine/Plugins/Runtime/SoundMod/Source/SoundMod/Classes/SoundModWave.h
+++ b/Engine/Plugins/Runtime/SoundMod/Source/SoundMod/Classes/SoundModWave.h
@@ -19,7 +19,7 @@ class USoundModWave : public USoundWaveStreaming
 	xmp_context xmpContext;
 
 	// Begin USoundWave interface.
-	virtual int32 GeneratePCMData(uint8* PCMData, const int32 SamplesNeeded) OVERRIDE;
+	virtual int32 GeneratePCMData(uint8* PCMData, const int32 SamplesNeeded) override;
 	// End USoundWave interface.
 
 };
\ No newline at end of file
diff --git a/Engine/Plugins/Runtime/SoundMod/Source/SoundMod/Private/SoundModPlugin.cpp b/Engine/Plugins/Runtime/SoundMod/Source/SoundMod/Private/SoundModPlugin.cpp
--- a/Engine/Plugins/Runtime/SoundMod/Source/SoundMod/Private/SoundModPlugin.cpp
+++ b/Engine/Plugins/Runtime/SoundMod/Source/SoundMod/Private/SoundModPlugin.cpp
@@ -6,8 +6,8 @@ class FSoundModPlugin : public ISoundModPlugin
 {
 public:
 	// IModuleInterface implementation
-	virtual void StartupModule() OVERRIDE;
-	virtual void ShutdownModule() OVERRIDE;
+	virtual void StartupModule() override;
+	virtual void ShutdownModule() override;
 	// End of IModuleInterface implementation
 };
 
diff --git a/Engine/Plugins/Runtime/SoundMod/Source/SoundMod/Public/SoundNodeModPlayer.h b/Engine/Plugins/Runtime/SoundMod/Source/SoundMod/Public/SoundNodeModPlayer.h
--- a/Engine/Plugins/Runtime/SoundMod/Source/SoundMod/Public/SoundNodeModPlayer.h
+++ b/Engine/Plugins/Runtime/SoundMod/Source/SoundMod/Public/SoundNodeModPlayer.h
@@ -21,12 +21,12 @@ class USoundNodeModPlayer : public USoundNode
 
 public:	
 	// Begin USoundNode Interface
-	virtual int32 GetMaxChildNodes() const OVERRIDE;
-	virtual float GetDuration() OVERRIDE;
-	virtual void ParseNodes( FAudioDevice* AudioDevice, const UPTRINT NodeWaveInstanceHash, FActiveSound& ActiveSound, const FSoundParseParameters& ParseParams, TArray<FWaveInstance*>& WaveInstances ) OVERRIDE;
-	virtual FString GetUniqueString() const OVERRIDE;
+	virtual int32 GetMaxChildNodes() const override;
+	virtual float GetDuration() override;
+	virtual void ParseNodes(FAudioDevice* AudioDevice, const UPTRINT NodeWaveInstanceHash, FActiveSound& ActiveSound, const FSoundParseParameters& ParseParams, TArray<FWaveInstance*>& WaveInstances) override;
+	virtual FString GetUniqueString() const override;
 #if WITH_EDITOR
-	virtual FString GetTitle() const OVERRIDE;
+	virtual FString GetTitle() const override;
 #endif
 	// End USoundNode Interface
 
diff --git a/Engine/Plugins/Runtime/SoundMod/Source/SoundModImporter/Classes/SoundModImporterFactory.h b/Engine/Plugins/Runtime/SoundMod/Source/SoundModImporter/Classes/SoundModImporterFactory.h
--- a/Engine/Plugins/Runtime/SoundMod/Source/SoundModImporter/Classes/SoundModImporterFactory.h
+++ b/Engine/Plugins/Runtime/SoundMod/Source/SoundModImporter/Classes/SoundModImporterFactory.h
@@ -11,7 +11,7 @@ class USoundModImporterFactory : public UFactory
 	GENERATED_UCLASS_BODY()
 
 	// UFactory interface
-	virtual UObject* FactoryCreateBinary(UClass* Class, UObject* InParent, FName Name, EObjectFlags Flags, UObject* Context, const TCHAR* Type, const uint8*& Buffer, const uint8* BufferEnd, FFeedbackContext* Warn) OVERRIDE;
+	virtual UObject* FactoryCreateBinary(UClass* Class, UObject* InParent, FName Name, EObjectFlags Flags, UObject* Context, const TCHAR* Type, const uint8*& Buffer, const uint8* BufferEnd, FFeedbackContext* Warn) override;
 	// End of UFactory interface
 };
 
diff --git a/Engine/Plugins/Runtime/SoundMod/Source/SoundModImporter/Private/AssetTypeActions_SoundMod.h b/Engine/Plugins/Runtime/SoundMod/Source/SoundModImporter/Private/AssetTypeActions_SoundMod.h
--- a/Engine/Plugins/Runtime/SoundMod/Source/SoundModImporter/Private/AssetTypeActions_SoundMod.h
+++ b/Engine/Plugins/Runtime/SoundMod/Source/SoundModImporter/Private/AssetTypeActions_SoundMod.h
@@ -8,14 +8,14 @@ class FAssetTypeActions_SoundMod: public FAssetTypeActions_Base
 {
 public:
 	// IAssetTypeActions Implementation
-	virtual FText GetName() const OVERRIDE{ return NSLOCTEXT("AssetTypeActions", "AssetTypeActions_SoundMod", "Sound Mod"); }
-	virtual FColor GetTypeColor() const OVERRIDE{ return FColor(255, 175, 0); }
-	virtual UClass* GetSupportedClass() const OVERRIDE;
-	virtual bool HasActions(const TArray<UObject*>& InObjects) const OVERRIDE{ return true; }
-	virtual void GetActions(const TArray<UObject*>& InObjects, FMenuBuilder& MenuBuilder) OVERRIDE;
-	virtual void AssetsActivated(const TArray<UObject*>& InObjects, EAssetTypeActivationMethod::Type ActivationType) OVERRIDE;
-	virtual uint32 GetCategories() OVERRIDE{ return EAssetTypeCategories::Sounds; }
-	virtual bool CanFilter() OVERRIDE{ return false; }
+	virtual FText GetName() const override{ return NSLOCTEXT("AssetTypeActions", "AssetTypeActions_SoundMod", "Sound Mod"); }
+	virtual FColor GetTypeColor() const override{ return FColor(255, 175, 0); }
+	virtual UClass* GetSupportedClass() const override;
+	virtual bool HasActions(const TArray<UObject*>& InObjects) const override{ return true; }
+	virtual void GetActions(const TArray<UObject*>& InObjects, FMenuBuilder& MenuBuilder) override;
+	virtual void AssetsActivated(const TArray<UObject*>& InObjects, EAssetTypeActivationMethod::Type ActivationType) override;
+	virtual uint32 GetCategories() override{ return EAssetTypeCategories::Sounds; }
+	virtual bool CanFilter() override { return false; }
 
 private:
 	/** Handler for when PlaySound is selected */
diff --git a/Engine/Plugins/Runtime/SoundMod/Source/SoundModImporter/Private/SoundModImporterModule.cpp b/Engine/Plugins/Runtime/SoundMod/Source/SoundModImporter/Private/SoundModImporterModule.cpp
--- a/Engine/Plugins/Runtime/SoundMod/Source/SoundModImporter/Private/SoundModImporterModule.cpp
+++ b/Engine/Plugins/Runtime/SoundMod/Source/SoundModImporter/Private/SoundModImporterModule.cpp
@@ -13,14 +13,14 @@ public:
 	TSharedPtr<FAssetTypeActions_SoundMod> SoundModAssetTypeActions;
 
 public:
-	virtual void StartupModule() OVERRIDE
+	virtual void StartupModule() override
 	{
 		// Register the asset type
 		SoundModAssetTypeActions = MakeShareable(new FAssetTypeActions_SoundMod);
 		FModuleManager::LoadModuleChecked<FAssetToolsModule>("AssetTools").Get().RegisterAssetTypeActions(SoundModAssetTypeActions.ToSharedRef());
 	}
 
-	virtual void ShutdownModule() OVERRIDE
+	virtual void ShutdownModule() override
 	{
 		if (!UObjectInitialized())
 		{
