Make USoundNode CreateInputPin and GetTiel return FText instead of FString

diff --git a/Engine/Plugins/Runtime/SoundMod/Source/SoundMod/Private/SoundNodeModPlayer.cpp b/Engine/Plugins/Runtime/SoundMod/Source/SoundMod/Private/SoundNodeModPlayer.cpp
--- a/Engine/Plugins/Runtime/SoundMod/Source/SoundMod/Private/SoundNodeModPlayer.cpp
+++ b/Engine/Plugins/Runtime/SoundMod/Source/SoundMod/Private/SoundNodeModPlayer.cpp
@@ -53,7 +53,7 @@ float USoundNodeModPlayer::GetDuration()
 }
 
 #if WITH_EDITOR
-FString USoundNodeModPlayer::GetTitle() const
+FText USoundNodeModPlayer::GetTitle() const
 {
 	FText SoundModName;
 	if (SoundMod)
@@ -65,18 +65,19 @@ FString USoundNodeModPlayer::GetTitle() const
 		SoundModName = LOCTEXT("NoSoundMod", "NONE");
 	}
 
-	FString Title;
+	FText Title;
+
+	FFormatNamedArguments Arguments;
+	Arguments.Add(TEXT("Description"), FText::FromString(Super::GetTitle()));
+	Arguments.Add(TEXT("SoundModName"), SoundModName);
 
 	if (bLooping)
 	{
-		FFormatNamedArguments Arguments;
-		Arguments.Add(TEXT("Description"), FText::FromString(Super::GetTitle()));
-		Arguments.Add(TEXT("SoundModName"), SoundModName);
-		Title = FText::Format(LOCTEXT("LoopingSoundWaveDescription", "Looping {Description} : {SoundModName}"), Arguments).ToString();
+		Title = FText::Format(LOCTEXT("LoopingSoundModDescription", "Looping {Description} : {SoundModName}"), Arguments);
 	}
 	else
 	{
-		Title = Super::GetTitle() + FString(TEXT(" : ")) + SoundModName.ToString();
+		Title = FText::Format(LOCTEXT("NonLoopingSoundModDescription", "{Description} : {SoundModName}"), Arguments);
 	}
 
 	return Title;
diff --git a/Engine/Plugins/Runtime/SoundMod/Source/SoundMod/Public/SoundNodeModPlayer.h b/Engine/Plugins/Runtime/SoundMod/Source/SoundMod/Public/SoundNodeModPlayer.h
--- a/Engine/Plugins/Runtime/SoundMod/Source/SoundMod/Public/SoundNodeModPlayer.h
+++ b/Engine/Plugins/Runtime/SoundMod/Source/SoundMod/Public/SoundNodeModPlayer.h
@@ -25,7 +25,7 @@ public:
 	virtual float GetDuration() override;
 	virtual void ParseNodes(FAudioDevice* AudioDevice, const UPTRINT NodeWaveInstanceHash, FActiveSound& ActiveSound, const FSoundParseParameters& ParseParams, TArray<FWaveInstance*>& WaveInstances) override;
 #if WITH_EDITOR
-	virtual FString GetTitle() const override;
+	virtual FText GetTitle() const override;
 #endif
 	// End USoundNode Interface
 
diff --git a/Engine/Source/Editor/UnrealEd/Private/SoundCueGraphNode.cpp b/Engine/Source/Editor/UnrealEd/Private/SoundCueGraphNode.cpp
--- a/Engine/Source/Editor/UnrealEd/Private/SoundCueGraphNode.cpp
+++ b/Engine/Source/Editor/UnrealEd/Private/SoundCueGraphNode.cpp
@@ -61,7 +61,7 @@ void USoundCueGraphNode::SetSoundNode(USoundNode* InSoundNode)
 
 void USoundCueGraphNode::CreateInputPin()
 {
-	UEdGraphPin* NewPin = CreatePin(EGPD_Input, TEXT("SoundNode"), TEXT(""), NULL, /*bIsArray=*/ false, /*bIsReference=*/ false, SoundNode->GetInputPinName(GetInputCount()));
+	UEdGraphPin* NewPin = CreatePin(EGPD_Input, TEXT("SoundNode"), TEXT(""), NULL, /*bIsArray=*/ false, /*bIsReference=*/ false, SoundNode->GetInputPinName(GetInputCount()).ToString());
 	if (NewPin->PinName.IsEmpty())
 	{
 		// Makes sure pin has a name for lookup purposes but user will never see it
@@ -145,7 +145,7 @@ FText USoundCueGraphNode::GetNodeTitle(ENodeTitleType::Type TitleType) const
 {
 	if (SoundNode)
 	{
-		return FText::FromString(SoundNode->GetTitle());
+		return SoundNode->GetTitle();
 	}
 	else
 	{
diff --git a/Engine/Source/Runtime/Engine/Classes/Sound/SoundNode.h b/Engine/Source/Runtime/Engine/Classes/Sound/SoundNode.h
--- a/Engine/Source/Runtime/Engine/Classes/Sound/SoundNode.h
+++ b/Engine/Source/Runtime/Engine/Classes/Sound/SoundNode.h
@@ -141,9 +141,9 @@ public:
 	virtual void SetChildNodes(TArray<USoundNode*>& InChildNodes);
 
 	/** Get the name of a specific input pin */
-	virtual FString GetInputPinName(int32 PinIndex) const { return TEXT(""); }
+	virtual FText GetInputPinName(int32 PinIndex) const { return FText::GetEmpty(); }
 
-	virtual FString GetTitle() const { return GetClass()->GetDescription(); }
+	virtual FText GetTitle() const { return GetClass()->GetDisplayNameText(); }
 
 	/** Helper function to set the position of a sound node on a grid */
 	void PlaceNode(int32 NodeColumn, int32 NodeRow, int32 RowCount );
diff --git a/Engine/Source/Runtime/Engine/Classes/Sound/SoundNodeBranch.h b/Engine/Source/Runtime/Engine/Classes/Sound/SoundNodeBranch.h
--- a/Engine/Source/Runtime/Engine/Classes/Sound/SoundNodeBranch.h
+++ b/Engine/Source/Runtime/Engine/Classes/Sound/SoundNodeBranch.h
@@ -45,8 +45,8 @@ public:
 	}
 
 #if WITH_EDITOR
-	virtual FString GetInputPinName(int32 PinIndex) const override;
-	virtual FString GetTitle() const override;
+	virtual FText GetInputPinName(int32 PinIndex) const override;
+	virtual FText GetTitle() const override;
 #endif //WITH_EDITOR
 	// End USoundNode interface.
 };
diff --git a/Engine/Source/Runtime/Engine/Classes/Sound/SoundNodeSwitch.h b/Engine/Source/Runtime/Engine/Classes/Sound/SoundNodeSwitch.h
--- a/Engine/Source/Runtime/Engine/Classes/Sound/SoundNodeSwitch.h
+++ b/Engine/Source/Runtime/Engine/Classes/Sound/SoundNodeSwitch.h
@@ -46,8 +46,8 @@ public:
 		RenamePins();
 	}
 
-	virtual FString GetInputPinName(int32 PinIndex) const override;
-	virtual FString GetTitle() const override;
+	virtual FText GetInputPinName(int32 PinIndex) const override;
+	virtual FText GetTitle() const override;
 #endif //WITH_EDITOR
 	// End USoundNode interface.
 };
diff --git a/Engine/Source/Runtime/Engine/Classes/Sound/SoundNodeWavePlayer.h b/Engine/Source/Runtime/Engine/Classes/Sound/SoundNodeWavePlayer.h
--- a/Engine/Source/Runtime/Engine/Classes/Sound/SoundNodeWavePlayer.h
+++ b/Engine/Source/Runtime/Engine/Classes/Sound/SoundNodeWavePlayer.h
@@ -25,7 +25,7 @@ public:
 	virtual float GetDuration() override;
 	virtual void ParseNodes( FAudioDevice* AudioDevice, const UPTRINT NodeWaveInstanceHash, FActiveSound& ActiveSound, const FSoundParseParameters& ParseParams, TArray<FWaveInstance*>& WaveInstances ) override;
 #if WITH_EDITOR
-	virtual FString GetTitle() const override;
+	virtual FText GetTitle() const override;
 #endif
 	// End USoundNode Interface
 
diff --git a/Engine/Source/Runtime/Engine/Private/SoundNodeBranch.cpp b/Engine/Source/Runtime/Engine/Private/SoundNodeBranch.cpp
--- a/Engine/Source/Runtime/Engine/Private/SoundNodeBranch.cpp
+++ b/Engine/Source/Runtime/Engine/Private/SoundNodeBranch.cpp
@@ -42,27 +42,31 @@ void USoundNodeBranch::CreateStartingConnectors()
 }
 
 #if WITH_EDITOR
-FString USoundNodeBranch::GetInputPinName(int32 PinIndex) const
+FText USoundNodeBranch::GetInputPinName(int32 PinIndex) const
 {
 	const BranchPurpose Branch = (BranchPurpose)PinIndex;
 	switch (Branch)
 	{
 	case BranchPurpose::ParameterTrue:
-		return LOCTEXT("True", "True").ToString();
+		return LOCTEXT("True", "True");
 
 	case BranchPurpose::ParameterFalse:
-		return LOCTEXT("False", "False").ToString();
+		return LOCTEXT("False", "False");
 
 	case BranchPurpose::ParameterUnset:
-		return LOCTEXT("ParamUnset", "Parameter Unset").ToString();
+		return LOCTEXT("ParamUnset", "Parameter Unset");
 	}
 
 	return Super::GetInputPinName(PinIndex);
 }
 
-FString USoundNodeBranch::GetTitle() const
+FText USoundNodeBranch::GetTitle() const
 {
-	return FString::Printf(TEXT("%s (%s)"), *Super::GetTitle(), *BoolParameterName.ToString());
+	FFormatNamedArguments Arguments;
+	Arguments.Add(TEXT("Description"), Super::GetTitle());
+	Arguments.Add(TEXT("ParameterName"), FText::FromName(BoolParameterName));
+
+	return FText::Format(LOCTEXT("Title", "{Description} ({ParameterName})"), Arguments);
 }
 #endif //WITH_EDITOR
 
diff --git a/Engine/Source/Runtime/Engine/Private/SoundNodeSwitch.cpp b/Engine/Source/Runtime/Engine/Private/SoundNodeSwitch.cpp
--- a/Engine/Source/Runtime/Engine/Private/SoundNodeSwitch.cpp
+++ b/Engine/Source/Runtime/Engine/Private/SoundNodeSwitch.cpp
@@ -57,28 +57,30 @@ void USoundNodeSwitch::RenamePins()
 	{
 		if (InputPins[i])
 		{
-			InputPins[i]->PinName = GetInputPinName(i);
+			InputPins[i]->PinName = GetInputPinName(i).ToString();
 		}
 	}
 }
 
-FString USoundNodeSwitch::GetInputPinName(int32 PinIndex) const
+FText USoundNodeSwitch::GetInputPinName(int32 PinIndex) const
 {
 	if (PinIndex == 0)
 	{
-		return LOCTEXT("ParamUnset", "Parameter Unset").ToString();
+		return LOCTEXT("ParamUnset", "Parameter Unset");
 	}
 	else
 	{
-		return FString::Printf(TEXT("%d"), PinIndex - 1);
+		return FText::FromString(FString::Printf(TEXT("%d"), PinIndex - 1));
 	}
-
-	return Super::GetInputPinName(PinIndex);
 }
 
-FString USoundNodeSwitch::GetTitle() const
+FText USoundNodeSwitch::GetTitle() const
 {
-	return FString::Printf(TEXT("%s (%s)"), *Super::GetTitle(), *IntParameterName.ToString());
+	FFormatNamedArguments Arguments;
+	Arguments.Add(TEXT("Description"), Super::GetTitle());
+	Arguments.Add(TEXT("ParameterName"), FText::FromName(IntParameterName));
+
+	return FText::Format(LOCTEXT("Title", "{Description} ({ParameterName})"), Arguments);
 }
 #endif //WITH_EDITOR
 
diff --git a/Engine/Source/Runtime/Engine/Private/SoundNodeWavePlayer.cpp b/Engine/Source/Runtime/Engine/Private/SoundNodeWavePlayer.cpp
--- a/Engine/Source/Runtime/Engine/Private/SoundNodeWavePlayer.cpp
+++ b/Engine/Source/Runtime/Engine/Private/SoundNodeWavePlayer.cpp
@@ -54,7 +54,7 @@ float USoundNodeWavePlayer::GetDuration()
 }
 
 #if WITH_EDITOR
-FString USoundNodeWavePlayer::GetTitle() const
+FText USoundNodeWavePlayer::GetTitle() const
 {
 	FText SoundWaveName;
 	if (SoundWave)
@@ -66,18 +66,18 @@ FString USoundNodeWavePlayer::GetTitle() const
 		SoundWaveName = LOCTEXT("NoSoundWave", "NONE");
 	}
 
-	FString Title;
+	FText Title;
 
+	FFormatNamedArguments Arguments;
+	Arguments.Add(TEXT("Description"), Super::GetTitle());
+	Arguments.Add(TEXT("SoundWaveName"), SoundWaveName);
 	if (bLooping)
 	{
-		FFormatNamedArguments Arguments;
-		Arguments.Add(TEXT("Description"), FText::FromString(Super::GetTitle()));
-		Arguments.Add(TEXT("SoundWaveName"), SoundWaveName);
-		Title = FText::Format(LOCTEXT("LoopingSoundWaveDescription", "Looping {Description} : {SoundWaveName}"), Arguments).ToString();
+		Title = FText::Format(LOCTEXT("LoopingSoundWaveDescription", "Looping {Description} : {SoundWaveName}"), Arguments);
 	}
 	else
 	{
-		Title = Super::GetTitle() + FString(TEXT(" : ")) + SoundWaveName.ToString();
+		Title = FText::Format(LOCTEXT("NonLoopingSoundWaveDescription", "{Description} : {SoundWaveName}"), Arguments);
 	}
 
 	return Title;
