Fix compile error

diff --git a/Engine/Plugins/Runtime/SoundMod/Source/SoundMod/Private/SoundNodeModPlayer.cpp b/Engine/Plugins/Runtime/SoundMod/Source/SoundMod/Private/SoundNodeModPlayer.cpp
--- a/Engine/Plugins/Runtime/SoundMod/Source/SoundMod/Private/SoundNodeModPlayer.cpp
+++ b/Engine/Plugins/Runtime/SoundMod/Source/SoundMod/Private/SoundNodeModPlayer.cpp
@@ -68,7 +68,7 @@ FText USoundNodeModPlayer::GetTitle() const
 	FText Title;
 
 	FFormatNamedArguments Arguments;
-	Arguments.Add(TEXT("Description"), FText::FromString(Super::GetTitle()));
+	Arguments.Add(TEXT("Description"), Super::GetTitle());
 	Arguments.Add(TEXT("SoundModName"), SoundModName);
 
 	if (bLooping)
