Fix non-unity
Submit missed file
Fix sound wave/mod not being cached after copy/paste of the player sound node
Fix PIE sound quality level property name

diff --git a/Engine/Plugins/Runtime/SoundMod/Source/SoundMod/Public/SoundNodeModPlayer.h b/Engine/Plugins/Runtime/SoundMod/Source/SoundMod/Public/SoundNodeModPlayer.h
--- a/Engine/Plugins/Runtime/SoundMod/Source/SoundMod/Public/SoundNodeModPlayer.h
+++ b/Engine/Plugins/Runtime/SoundMod/Source/SoundMod/Public/SoundNodeModPlayer.h
@@ -2,24 +2,39 @@
 
 
 #pragma once
-#include "Sound/SoundNode.h"
+#include "Sound/SoundNodeAssetReferencer.h"
 #include "SoundNodeModPlayer.generated.h"
 
 /** 
  * Sound node that contains a reference to the mod file to be played
  */
 UCLASS(hidecategories=Object, editinlinenew, MinimalAPI, meta=( DisplayName="Mod Player" ))
-class USoundNodeModPlayer : public USoundNode
+class USoundNodeModPlayer : public USoundNodeAssetReferencer
 {
-	GENERATED_UCLASS_BODY()
+	GENERATED_BODY()
+
+private:
+	UPROPERTY(EditAnywhere, Category=ModPlayer, meta=(DisplayName="Sound Mod"))
+	TAssetPtr<USoundMod> SoundModAssetPtr;
 
-	UPROPERTY(EditAnywhere, Category=ModPlayer)
 	USoundMod* SoundMod;
 
 	UPROPERTY(EditAnywhere, Category=ModPlayer)
 	uint32 bLooping:1;
 
 public:	
+
+	USoundMod* GetSoundMod() const { return SoundMod; }
+	void SetSoundMod(USoundMod* SoundMod);
+
+	// Begin UObject Interface
+	static void AddReferencedObjects(UObject* InThis, FReferenceCollector& Collector);
+
+#if WITH_EDITOR
+	virtual void PostEditChangeProperty(FPropertyChangedEvent& PropertyChangedEvent) override;
+#endif
+	// End UObject Interface
+
 	// Begin USoundNode Interface
 	virtual int32 GetMaxChildNodes() const override;
 	virtual float GetDuration() override;
@@ -29,5 +44,9 @@ public:
 #endif
 	// End USoundNode Interface
 
+	// Begin USoundNodeAssetReferencer Interface
+	virtual void LoadAsset() override;
+	// End USoundNode Interface
+
 };
 
diff --git a/Engine/Source/Editor/UnrealEd/Classes/Settings/LevelEditorPlaySettings.h b/Engine/Source/Editor/UnrealEd/Classes/Settings/LevelEditorPlaySettings.h
--- a/Engine/Source/Editor/UnrealEd/Classes/Settings/LevelEditorPlaySettings.h
+++ b/Engine/Source/Editor/UnrealEd/Classes/Settings/LevelEditorPlaySettings.h
@@ -163,7 +163,7 @@ public:
 
 	/** Which quality level to use when playing in editor */
 	UPROPERTY(config, EditAnywhere, Category=PlayInEditor)
-	int32 PlayInEditorSoundQuality;
+	int32 PlayInEditorSoundQualityLevel;
 
 	/** True if Play In Editor should only load currently-visible levels in PIE. */
 	UPROPERTY(config)
diff --git a/Engine/Source/Editor/UnrealEd/Private/Settings/LevelEditorPlaySettingsCustomization.h b/Engine/Source/Editor/UnrealEd/Private/Settings/LevelEditorPlaySettingsCustomization.h
--- a/Engine/Source/Editor/UnrealEd/Private/Settings/LevelEditorPlaySettingsCustomization.h
+++ b/Engine/Source/Editor/UnrealEd/Private/Settings/LevelEditorPlaySettingsCustomization.h
@@ -304,7 +304,7 @@ public:
 			PlayInEditorCategory.GetDefaultProperties(PIECategoryProperties, true, false);
 
 			TSharedPtr<IPropertyHandle> PIEEnableSoundHandle = LayoutBuilder.GetProperty(GET_MEMBER_NAME_CHECKED(ULevelEditorPlaySettings, EnableSound));
-			PIESoundQualityLevelHandle = LayoutBuilder.GetProperty(GET_MEMBER_NAME_CHECKED(ULevelEditorPlaySettings, PlayInEditorSoundQuality));
+			PIESoundQualityLevelHandle = LayoutBuilder.GetProperty(GET_MEMBER_NAME_CHECKED(ULevelEditorPlaySettings, PlayInEditorSoundQualityLevel));
 			PIESoundQualityLevelHandle->MarkHiddenByCustomization();
 
 			for (TSharedRef<IPropertyHandle>& PropertyHandle : PIECategoryProperties)
diff --git a/Engine/Source/Runtime/Engine/Classes/Sound/SoundNodeAssetReferencer.h b/Engine/Source/Runtime/Engine/Classes/Sound/SoundNodeAssetReferencer.h
--- a/Engine/Source/Runtime/Engine/Classes/Sound/SoundNodeAssetReferencer.h
+++ b/Engine/Source/Runtime/Engine/Classes/Sound/SoundNodeAssetReferencer.h
@@ -15,5 +15,9 @@ class ENGINE_API USoundNodeAssetReferencer : public USoundNode
 
 public:
 	virtual void LoadAsset() PURE_VIRTUAL(USoundNodesAssetReferencer::LoadAsset,);
+
+#if WITH_EDITOR
+	virtual void PostEditImport() override;
+#endif
 };
 
diff --git a/Engine/Source/Runtime/Engine/Private/AudioDeviceManager.cpp b/Engine/Source/Runtime/Engine/Private/AudioDeviceManager.cpp
--- a/Engine/Source/Runtime/Engine/Private/AudioDeviceManager.cpp
+++ b/Engine/Source/Runtime/Engine/Private/AudioDeviceManager.cpp
@@ -5,6 +5,8 @@
 #include "AudioDevice.h"
 #include "AudioDeviceManager.h"
 #include "Sound/SoundWave.h"
+#include "Sound/AudioSettings.h"
+#include "GameFramework/GameUserSettings.h"
 
 // Private consts for helping with index/generation determination in audio device manager
 static const uint32 AUDIO_DEVICE_HANDLE_INDEX_BITS		= 24;
diff --git a/Engine/Source/Runtime/Engine/Private/SoundNodeAssetReferencer.cpp b/Engine/Source/Runtime/Engine/Private/SoundNodeAssetReferencer.cpp
new file mode 100644
--- /dev/null
+++ b/Engine/Source/Runtime/Engine/Private/SoundNodeAssetReferencer.cpp
@@ -0,0 +1,13 @@
+// Copyright 1998-2015 Epic Games, Inc. All Rights Reserved.
+
+#include "EnginePrivate.h"
+#include "Sound/SoundNodeAssetReferencer.h"
+
+#if WITH_EDITOR
+void USoundNodeAssetReferencer::PostEditImport()
+{
+	Super::PostEditImport();
+
+	LoadAsset();
+}
+#endif
diff --git a/Engine/Source/Runtime/Engine/Private/SoundNodeQualityLevel.cpp b/Engine/Source/Runtime/Engine/Private/SoundNodeQualityLevel.cpp
--- a/Engine/Source/Runtime/Engine/Private/SoundNodeQualityLevel.cpp
+++ b/Engine/Source/Runtime/Engine/Private/SoundNodeQualityLevel.cpp
@@ -4,6 +4,13 @@
 #include "EnginePrivate.h"
 #include "ActiveSound.h"
 #include "Sound/SoundNodeQualityLevel.h"
+#include "Sound/AudioSettings.h"
+
+#if WITH_EDITORONLY_DATA
+#include "Editor.h"
+#include "EdGraph/EdGraph.h"
+#include "Settings/LevelEditorPlaySettings.h"
+#endif
 
 #if WITH_EDITOR
 void USoundNodeQualityLevel::PostLoad()
@@ -63,7 +70,7 @@ void USoundNodeQualityLevel::ParseNodes( FAudioDevice* AudioDevice, const UPTRIN
 			const bool bIsPIESound = ((GEditor->bIsSimulatingInEditor || GEditor->PlayWorld != NULL) && ActiveSound.World != NULL);
 			if (bIsPIESound)
 			{
-				QualityLevel = GetDefault<ULevelEditorPlaySettings>()->PlayInEditorSoundQuality;
+				QualityLevel = GetDefault<ULevelEditorPlaySettings>()->PlayInEditorSoundQualityLevel;
 			}
 		}
 	}
