Make sound waves and mods load correctly when async loading.

diff --git a/Engine/Plugins/Runtime/SoundMod/Source/SoundMod/Private/SoundNodeModPlayer.cpp b/Engine/Plugins/Runtime/SoundMod/Source/SoundMod/Private/SoundNodeModPlayer.cpp
--- a/Engine/Plugins/Runtime/SoundMod/Source/SoundMod/Private/SoundNodeModPlayer.cpp
+++ b/Engine/Plugins/Runtime/SoundMod/Source/SoundMod/Private/SoundNodeModPlayer.cpp
@@ -16,7 +16,26 @@ void USoundNodeModPlayer::AddReferencedObjects(UObject* InThis, FReferenceCollec
 
 void USoundNodeModPlayer::LoadAsset()
 {
-	SoundMod = SoundModAssetPtr.LoadSynchronous();
+	if (IsAsyncLoading())
+	{
+		SoundMod = SoundModAssetPtr.Get();
+		if (SoundMod == nullptr)
+		{
+			LoadPackageAsync(SoundModAssetPtr.GetLongPackageName(), FLoadPackageAsyncDelegate::CreateUObject(this, &USoundNodeModPlayer::OnSoundModLoaded));
+		}
+	}
+	else
+	{
+		SoundMod = SoundModAssetPtr.LoadSynchronous();
+	}
+}
+
+void USoundNodeModPlayer::OnSoundModLoaded(const FName& PackageName, UPackage * Package, EAsyncLoadingResult::Type Result)
+{
+	if (Result == EAsyncLoadingResult::Succeeded)
+	{
+		SoundMod = SoundModAssetPtr.Get();
+	}
 }
 
 void USoundNodeModPlayer::SetSoundMod(USoundMod* InSoundMod)
diff --git a/Engine/Plugins/Runtime/SoundMod/Source/SoundMod/Public/SoundNodeModPlayer.h b/Engine/Plugins/Runtime/SoundMod/Source/SoundMod/Public/SoundNodeModPlayer.h
--- a/Engine/Plugins/Runtime/SoundMod/Source/SoundMod/Public/SoundNodeModPlayer.h
+++ b/Engine/Plugins/Runtime/SoundMod/Source/SoundMod/Public/SoundNodeModPlayer.h
@@ -22,6 +22,8 @@ private:
 	UPROPERTY(EditAnywhere, Category=ModPlayer)
 	uint32 bLooping:1;
 
+	void OnSoundModLoaded(const FName& PackageName, UPackage * Package, EAsyncLoadingResult::Type Result);
+
 public:	
 
 	USoundMod* GetSoundMod() const { return SoundMod; }
diff --git a/Engine/Source/Runtime/Engine/Classes/Sound/SoundNodeWavePlayer.h b/Engine/Source/Runtime/Engine/Classes/Sound/SoundNodeWavePlayer.h
--- a/Engine/Source/Runtime/Engine/Classes/Sound/SoundNodeWavePlayer.h
+++ b/Engine/Source/Runtime/Engine/Classes/Sound/SoundNodeWavePlayer.h
@@ -19,6 +19,8 @@ private:
 
 	USoundWave* SoundWave;
 
+	void OnSoundWaveLoaded(const FName& PackageName, UPackage * Package, EAsyncLoadingResult::Type Result);
+
 public:	
 
 	UPROPERTY(EditAnywhere, Category=WavePlayer)
diff --git a/Engine/Source/Runtime/Engine/Private/SoundNodeWavePlayer.cpp b/Engine/Source/Runtime/Engine/Private/SoundNodeWavePlayer.cpp
--- a/Engine/Source/Runtime/Engine/Private/SoundNodeWavePlayer.cpp
+++ b/Engine/Source/Runtime/Engine/Private/SoundNodeWavePlayer.cpp
@@ -17,7 +17,26 @@ void USoundNodeWavePlayer::AddReferencedObjects(UObject* InThis, FReferenceColle
 
 void USoundNodeWavePlayer::LoadAsset()
 {
-	SoundWave = SoundWaveAssetPtr.LoadSynchronous();
+	if (IsAsyncLoading())
+	{
+		SoundWave = SoundWaveAssetPtr.Get();
+		if (SoundWave == nullptr)
+		{
+			LoadPackageAsync(SoundWaveAssetPtr.GetLongPackageName(), FLoadPackageAsyncDelegate::CreateUObject(this, &USoundNodeWavePlayer::OnSoundWaveLoaded));
+		}
+	}
+	else
+	{
+		SoundWave = SoundWaveAssetPtr.LoadSynchronous();
+	}
+}
+
+void USoundNodeWavePlayer::OnSoundWaveLoaded(const FName& PackageName, UPackage * Package, EAsyncLoadingResult::Type Result)
+{
+	if (Result == EAsyncLoadingResult::Succeeded)
+	{
+		SoundWave = SoundWaveAssetPtr.Get();
+	}
 }
 
 void USoundNodeWavePlayer::SetSoundWave(USoundWave* InSoundWave)
