Validate there is a long package name (i.e. the SoundAssetPtr points to something) before trying to async load the pointer

diff --git a/Engine/Plugins/Runtime/SoundMod/Source/SoundMod/Private/SoundNodeModPlayer.cpp b/Engine/Plugins/Runtime/SoundMod/Source/SoundMod/Private/SoundNodeModPlayer.cpp
--- a/Engine/Plugins/Runtime/SoundMod/Source/SoundMod/Private/SoundNodeModPlayer.cpp
+++ b/Engine/Plugins/Runtime/SoundMod/Source/SoundMod/Private/SoundNodeModPlayer.cpp
@@ -21,7 +21,11 @@ void USoundNodeModPlayer::LoadAsset()
 		SoundMod = SoundModAssetPtr.Get();
 		if (SoundMod == nullptr)
 		{
-			LoadPackageAsync(SoundModAssetPtr.GetLongPackageName(), FLoadPackageAsyncDelegate::CreateUObject(this, &USoundNodeModPlayer::OnSoundModLoaded));
+			const FString LongPackageName = SoundModAssetPtr.GetLongPackageName();
+			if (!LongPackageName.IsEmpty())
+			{
+				LoadPackageAsync(LongPackageName, FLoadPackageAsyncDelegate::CreateUObject(this, &USoundNodeModPlayer::OnSoundModLoaded));
+			}
 		}
 	}
 	else
diff --git a/Engine/Source/Runtime/Engine/Private/SoundNodeWavePlayer.cpp b/Engine/Source/Runtime/Engine/Private/SoundNodeWavePlayer.cpp
--- a/Engine/Source/Runtime/Engine/Private/SoundNodeWavePlayer.cpp
+++ b/Engine/Source/Runtime/Engine/Private/SoundNodeWavePlayer.cpp
@@ -22,7 +22,11 @@ void USoundNodeWavePlayer::LoadAsset()
 		SoundWave = SoundWaveAssetPtr.Get();
 		if (SoundWave == nullptr)
 		{
-			LoadPackageAsync(SoundWaveAssetPtr.GetLongPackageName(), FLoadPackageAsyncDelegate::CreateUObject(this, &USoundNodeWavePlayer::OnSoundWaveLoaded));
+			const FString LongPackageName = SoundWaveAssetPtr.GetLongPackageName();
+			if (!LongPackageName.IsEmpty())
+			{
+				LoadPackageAsync(LongPackageName, FLoadPackageAsyncDelegate::CreateUObject(this, &USoundNodeWavePlayer::OnSoundWaveLoaded));
+			}
 		}
 	}
 	else
