Fix compilation of SoundModWave.h

diff --git a/Engine/Plugins/Runtime/SoundMod/Source/SoundMod/Classes/SoundModWave.h b/Engine/Plugins/Runtime/SoundMod/Source/SoundMod/Classes/SoundModWave.h
--- a/Engine/Plugins/Runtime/SoundMod/Source/SoundMod/Classes/SoundModWave.h
+++ b/Engine/Plugins/Runtime/SoundMod/Source/SoundMod/Classes/SoundModWave.h
@@ -6,6 +6,7 @@
 
 #pragma once
 #include "xmp.h"
+#include "Sound/SoundWaveProcedural.h"
 #include "SoundModWave.generated.h"
 
 UCLASS()
