Fix missing constructor consolidation reference in SoundMod from last submit

diff --git a/Engine/Plugins/Runtime/SoundMod/Source/SoundMod/Private/SoundMod.cpp b/Engine/Plugins/Runtime/SoundMod/Source/SoundMod/Private/SoundMod.cpp
--- a/Engine/Plugins/Runtime/SoundMod/Source/SoundMod/Private/SoundMod.cpp
+++ b/Engine/Plugins/Runtime/SoundMod/Source/SoundMod/Private/SoundMod.cpp
@@ -27,9 +27,7 @@ void USoundMod::Parse(class FAudioDevice* AudioDevice, const UPTRINT NodeWaveIns
 		const int32 SampleRate = 44100;
 
 		// Create a new wave instance and associate with the ActiveSound
-		WaveInstance = new FWaveInstance(&ActiveSound);
-		WaveInstance->WaveInstanceHash = NodeWaveInstanceHash;
-		ActiveSound.WaveInstances.Add(NodeWaveInstanceHash, WaveInstance);
+		WaveInstance = new FWaveInstance(NodeWaveInstanceHash, ActiveSound);
 
 		// Create streaming wave object
 		USoundModWave* ModWave = NewObject<USoundModWave>();
