Fix strict vs2019 compile

diff --git a/Engine/Plugins/Runtime/SoundMod/Source/SoundModImporter/Private/SoundModImporterFactory.cpp b/Engine/Plugins/Runtime/SoundMod/Source/SoundModImporter/Private/SoundModImporterFactory.cpp
--- a/Engine/Plugins/Runtime/SoundMod/Source/SoundModImporter/Private/SoundModImporterFactory.cpp
+++ b/Engine/Plugins/Runtime/SoundMod/Source/SoundModImporter/Private/SoundModImporterFactory.cpp
@@ -1,11 +1,13 @@
 // Copyright Epic Games, Inc. All Rights Reserved.
 
 #include "SoundModImporterFactory.h"
-#include "EngineGlobals.h"
+
+#include "AudioDeviceManager.h"
+#include "Components/AudioComponent.h"
 #include "Editor.h"
+#include "EngineGlobals.h"
 #include "SoundMod.h"
 #include "xmp.h"
-#include "Components/AudioComponent.h"
 
 //////////////////////////////////////////////////////////////////////////
 // USoundModImporterFactory
@@ -44,12 +46,15 @@ UObject* USoundModImporterFactory::FactoryCreateBinary
 
 	// TODO - Audio Threading. This needs to be sent to the audio device and wait on stopping the sounds
 	TArray<UAudioComponent*> ComponentsToRestart;
-	FAudioDeviceManager* AudioDeviceManager = FAudioDeviceManager::Get();
-	if (AudioDeviceManager && ExistingSound)
-	{
-		// TODO: Generalize the stop sounds function
-		//AudioDeviceManager->StopSoundsForReimport(ExistingSound, ComponentsToRestart);
-	}
+
+// 	if (FAudioDeviceManager* AudioDeviceManager = FAudioDeviceManager::Get())
+// 	{
+// 		if (ExistingSound)
+// 		{
+// 			// TODO: Generalize the stop sounds function
+// 			AudioDeviceManager->StopSoundsForReimport(ExistingSound, ComponentsToRestart);
+// 		}
+// 	}
 
 	bool bUseExistingSettings = bSoundModFactorySuppressImportOverwriteDialog;
 
diff --git a/Engine/Source/Editor/AudioEditor/Classes/Factories/SoundSurroundFactory.h b/Engine/Source/Editor/AudioEditor/Classes/Factories/SoundSurroundFactory.h
--- a/Engine/Source/Editor/AudioEditor/Classes/Factories/SoundSurroundFactory.h
+++ b/Engine/Source/Editor/AudioEditor/Classes/Factories/SoundSurroundFactory.h
@@ -1,27 +1,13 @@
 // Copyright Epic Games, Inc. All Rights Reserved.
-
-//~=============================================================================
-// SoundSurroundFactory
-//~=============================================================================
-
 #pragma once
 
 #include "CoreMinimal.h"
-#include "UObject/ObjectMacros.h"
+
 #include "Factories/Factory.h"
+#include "UObject/ObjectMacros.h"
+
 #include "SoundSurroundFactory.generated.h"
 
-const FString SurroundSpeakerLocations[SPEAKER_Count] =
-{
-	TEXT("_fl"),		// SPEAKER_FrontLeft
-	TEXT("_fr"),		// SPEAKER_FrontRight
-	TEXT("_fc"),		// SPEAKER_FrontCenter
-	TEXT("_lf"),		// SPEAKER_LowFrequency
-	TEXT("_sl"),		// SPEAKER_SideLeft
-	TEXT("_sr"),		// SPEAKER_SideRight
-	TEXT("_bl"),		// SPEAKER_BackLeft
-	TEXT("_br")			// SPEAKER_BackRight
-};
 
 UCLASS(MinimalAPI, hidecategories=Object)
 class USoundSurroundFactory : public UFactory
@@ -31,7 +17,6 @@ class USoundSurroundFactory : public UFactory
 	UPROPERTY()
 	float CueVolume;
 
-
 	//~ Begin UFactory Interface
 	virtual UObject* FactoryCreateBinary( UClass* Class, UObject* InParent, FName Name, EObjectFlags Flags, UObject* Context, const TCHAR* Type, const uint8*& Buffer, const uint8* BufferEnd, FFeedbackContext* Warn ) override;
 	virtual bool FactoryCanImport( const FString& Filename ) override;
diff --git a/Engine/Source/Editor/AudioEditor/Private/Factories/ReimportSoundSurroundFactory.cpp b/Engine/Source/Editor/AudioEditor/Private/Factories/ReimportSoundSurroundFactory.cpp
--- a/Engine/Source/Editor/AudioEditor/Private/Factories/ReimportSoundSurroundFactory.cpp
+++ b/Engine/Source/Editor/AudioEditor/Private/Factories/ReimportSoundSurroundFactory.cpp
@@ -1,14 +1,16 @@
 // Copyright Epic Games, Inc. All Rights Reserved.
-
 #include "Factories/ReimportSoundSurroundFactory.h"
-#include "Logging/MessageLog.h"
-#include "Sound/SoundWave.h"
-#include "Misc/Paths.h"
+
+#include "Audio.h"
 #include "EditorFramework/AssetImportData.h"
 #include "HAL/FileManager.h"
+#include "Logging/MessageLog.h"
+#include "Misc/Paths.h"
+#include "Sound/SoundWave.h"
 
 #define LOCTEXT_NAMESPACE "AudioEditorFactories"
 
+
 UReimportSoundSurroundFactory::UReimportSoundSurroundFactory(const FObjectInitializer& ObjectInitializer)
 	: Super(ObjectInitializer)
 {
diff --git a/Engine/Source/Editor/AudioEditor/Private/Factories/SoundSurroundFactory.cpp b/Engine/Source/Editor/AudioEditor/Private/Factories/SoundSurroundFactory.cpp
--- a/Engine/Source/Editor/AudioEditor/Private/Factories/SoundSurroundFactory.cpp
+++ b/Engine/Source/Editor/AudioEditor/Private/Factories/SoundSurroundFactory.cpp
@@ -1,6 +1,7 @@
 // Copyright Epic Games, Inc. All Rights Reserved.
-
 #include "Factories/SoundSurroundFactory.h"
+
+#include "Audio.h"
 #include "PackageTools.h"
 #include "Sound/SoundWave.h"
 #include "Misc/Paths.h"
@@ -9,6 +10,7 @@
 #include "Editor.h"
 #include "EditorFramework/AssetImportData.h"
 
+
 USoundSurroundFactory::USoundSurroundFactory(const FObjectInitializer& ObjectInitializer)
 	: Super(ObjectInitializer)
 {
diff --git a/Engine/Source/Runtime/Engine/Public/Audio.h b/Engine/Source/Runtime/Engine/Public/Audio.h
--- a/Engine/Source/Runtime/Engine/Public/Audio.h
+++ b/Engine/Source/Runtime/Engine/Public/Audio.h
@@ -55,16 +55,23 @@ DECLARE_CYCLE_STAT_EXTERN(TEXT("Decompress Streamed"), STAT_AudioStreamedDecompr
 DECLARE_CYCLE_STAT_EXTERN(TEXT("Buffer Creation"), STAT_AudioResourceCreationTime, STATGROUP_Audio, );
 
 
-class FAudioDevice;
-class USoundNode;
-struct FSoundModulationControls;
-class USoundWave;
+// Forward declarations.
 class USoundClass;
-class USoundSubmix;
+class UAudioComponent;
+class USoundNode;
 class USoundSourceBus;
+class USoundSubmix;
+class USoundWave;
+
+class FAudioDevice;
+
 struct FActiveSound;
-struct FWaveInstance;
+struct FReverbSettings;
+struct FSampleLoop;
+struct FSoundModulationControls;
 struct FSoundSourceBusSendInfo;
+struct FWaveInstance;
+
 
 /**
  * Channel definitions for multistream waves
@@ -85,12 +92,17 @@ enum EAudioSpeakers
 	SPEAKER_Count
 };
 
-// Forward declarations.
-class UAudioComponent;
-class USoundNode;
-struct FWaveInstance;
-struct FReverbSettings;
-struct FSampleLoop;
+const FString SurroundSpeakerLocations[SPEAKER_Count] =
+{
+	TEXT("_fl"),	// SPEAKER_FrontLeft
+	TEXT("_fr"),	// SPEAKER_FrontRight
+	TEXT("_fc"),	// SPEAKER_FrontCenter
+	TEXT("_lf"),	// SPEAKER_LowFrequency
+	TEXT("_sl"),	// SPEAKER_SideLeft
+	TEXT("_sr"),	// SPEAKER_SideRight
+	TEXT("_bl"),	// SPEAKER_BackLeft
+	TEXT("_br")		// SPEAKER_BackRight
+};
 
 enum ELoopingMode
 {
