Fix for loctext collision
#rnx

diff --git a/Engine/Plugins/Runtime/SoundMod/Source/SoundModImporter/Private/AssetTypeActions_SoundMod.cpp b/Engine/Plugins/Runtime/SoundMod/Source/SoundModImporter/Private/AssetTypeActions_SoundMod.cpp
--- a/Engine/Plugins/Runtime/SoundMod/Source/SoundModImporter/Private/AssetTypeActions_SoundMod.cpp
+++ b/Engine/Plugins/Runtime/SoundMod/Source/SoundModImporter/Private/AssetTypeActions_SoundMod.cpp
@@ -5,7 +5,7 @@
 #include "SoundMod.h"
 #include "ToolMenuSection.h"
 
-#define LOCTEXT_NAMESPACE "AssetTypeActions"
+#define LOCTEXT_NAMESPACE "SoundModImporter"
 
 UClass* FAssetTypeActions_SoundMod::GetSupportedClass() const
 {
